Chooser evaluation for AnimNext

diff --git a/Engine/Plugins/Experimental/AnimNextChooser/AnimNextChooser.uplugin b/Engine/Plugins/Experimental/AnimNextChooser/AnimNextChooser.uplugin
new file mode 100644
--- /dev/null
+++ b/Engine/Plugins/Experimental/AnimNextChooser/AnimNextChooser.uplugin
@@ -0,0 +1,37 @@
+{
+	"FileVersion" : 3,
+	"Version" : 1,
+	"VersionName" : "0.1",
+	"FriendlyName" : "Anim Next Chooser",
+	"Description" : "Chooser integration for Anim Next.",
+	"Category" : "Animation",
+	"CreatedBy" : "Epic Games, Inc.",
+	"CreatedByURL" : "https://epicgames.com",
+	"DocsURL" : "",
+	"MarketplaceURL" : "",
+	"SupportURL" : "",
+	"EnabledByDefault" : false,
+	"CanContainContent" : false,
+	"IsBetaVersion": true,
+	"IsExperimentalVersion": true,
+	"Installed" : false,
+	"Plugins":
+	[
+		{
+			"Name": "AnimNext",
+			"Enabled": true
+		},
+		{
+			"Name": "Chooser",
+			"Enabled": true
+		}
+	],
+	"Modules" :
+	[
+		{
+			"Name" : "AnimNextChooser",
+			"Type" : "Runtime",
+			"LoadingPhase" : "Default"
+		},
+	]
+}
diff --git a/Engine/Plugins/Experimental/AnimNextChooser/Source/AnimNextChooser/AnimNextChooser.Build.cs b/Engine/Plugins/Experimental/AnimNextChooser/Source/AnimNextChooser/AnimNextChooser.Build.cs
new file mode 100644
--- /dev/null
+++ b/Engine/Plugins/Experimental/AnimNextChooser/Source/AnimNextChooser/AnimNextChooser.Build.cs
@@ -0,0 +1,23 @@
+// Copyright Epic Games, Inc. All Rights Reserved.
+
+namespace UnrealBuildTool.Rules
+{
+	public class AnimNextChooser : ModuleRules
+	{
+		public AnimNextChooser(ReadOnlyTargetRules Target) : base(Target)
+		{
+			PublicDependencyModuleNames.AddRange(
+				new string[]
+				{
+					"Core",
+					"CoreUObject",
+					"RigVM",
+					"ControlRig",
+					"Engine",
+					"AnimNext",
+					"Chooser",
+				}
+			);
+		}
+	}
+}
\ No newline at end of file
diff --git a/Engine/Plugins/Experimental/AnimNextChooser/Source/AnimNextChooser/Private/ChooserParameters.cpp b/Engine/Plugins/Experimental/AnimNextChooser/Source/AnimNextChooser/Private/ChooserParameters.cpp
new file mode 100644
--- /dev/null
+++ b/Engine/Plugins/Experimental/AnimNextChooser/Source/AnimNextChooser/Private/ChooserParameters.cpp
@@ -0,0 +1,60 @@
+// Copyright Epic Games, Inc. All Rights Reserved.
+#include "ChooserParameters.h"
+#include "Param/ParamStack.h"
+
+#if WITH_EDITOR
+#include "StructUtils/PropertyBag.h"
+#endif
+
+bool FBoolAnimProperty::GetValue(FChooserEvaluationContext& Context, bool& OutResult) const
+{
+	UE::AnimNext::FParamResult Result;
+	UE::AnimNext::FParamStack& ParamStack = UE::AnimNext::FParamStack::Get();
+	OutResult = ParamStack.GetParamValue<bool>(UE::AnimNext::FParamId(PropertyName), &Result);
+	return Result.IsSuccessful();
+}
+
+bool FBoolAnimProperty::SetValue(FChooserEvaluationContext& Context, bool InValue) const
+{
+	return false;
+}
+
+void FBoolAnimProperty::GetDisplayName(FText& OutName) const
+{
+	FString DisplayName = PropertyName.ToString();
+	int Index = -1;
+	DisplayName.FindLastChar(':', Index);
+	if (Index>=0)
+	{
+		DisplayName = DisplayName.RightChop(Index + 1);
+	}
+	
+	OutName = FText::FromString(DisplayName);
+}
+
+
+bool FFloatAnimProperty::GetValue(FChooserEvaluationContext& Context, double& OutResult) const
+{
+	UE::AnimNext::FParamResult Result;
+	UE::AnimNext::FParamStack& ParamStack = UE::AnimNext::FParamStack::Get();
+	OutResult = ParamStack.GetParamValue<double>(UE::AnimNext::FParamId(PropertyName), &Result);
+	return Result.IsSuccessful();
+}
+
+bool FFloatAnimProperty::SetValue(FChooserEvaluationContext& Context, double InValue) const
+{
+	return false;
+}
+
+void FFloatAnimProperty::GetDisplayName(FText& OutName) const
+{
+	FString DisplayName = PropertyName.ToString();
+	int Index = -1;
+	DisplayName.FindLastChar(':', Index);
+	if (Index>=0)
+	{
+		DisplayName = DisplayName.RightChop(Index + 1);
+	}
+	
+	OutName = FText::FromString(DisplayName);
+}
\ No newline at end of file
diff --git a/Engine/Plugins/Experimental/AnimNextChooser/Source/AnimNextChooser/Private/ChooserParameters.h b/Engine/Plugins/Experimental/AnimNextChooser/Source/AnimNextChooser/Private/ChooserParameters.h
new file mode 100644
--- /dev/null
+++ b/Engine/Plugins/Experimental/AnimNextChooser/Source/AnimNextChooser/Private/ChooserParameters.h
@@ -0,0 +1,35 @@
+// Copyright Epic Games, Inc. All Rights Reserved.
+#pragma once
+
+#include "CoreMinimal.h"
+#include "IChooserParameterBool.h"
+#include "IChooserParameterFloat.h"
+#include "ChooserParameters.generated.h"
+
+USTRUCT(DisplayName = "Bool Anim Param")
+struct FBoolAnimProperty :  public FChooserParameterBoolBase
+{
+	GENERATED_BODY()
+public:
+
+	UPROPERTY(DisplayName = "Parameter", EditAnywhere, Category="Parameter", meta=(CustomWidget = "ParamName", AllowedParamType = "bool"))
+	FName PropertyName;
+	
+	virtual bool GetValue(FChooserEvaluationContext& Context, bool& OutResult) const override;
+	virtual bool SetValue(FChooserEvaluationContext& Context, bool InValue) const override;
+	virtual void GetDisplayName(FText& OutName) const override;
+};
+
+USTRUCT(DisplayName = "Float Anim Param")
+struct FFloatAnimProperty :  public FChooserParameterFloatBase
+{
+	GENERATED_BODY()
+public:
+
+	UPROPERTY(DisplayName = "Parameter", EditAnywhere, Category="Parameter", meta=(CustomWidget = "ParamName", AllowedParamType = "double"))
+	FName PropertyName;
+	
+	virtual bool GetValue(FChooserEvaluationContext& Context, double& OutResult) const override;
+	virtual bool SetValue(FChooserEvaluationContext& Context, double InValue) const override;
+	virtual void GetDisplayName(FText& OutName) const override;
+};
\ No newline at end of file
diff --git a/Engine/Plugins/Experimental/AnimNextChooser/Source/AnimNextChooser/Private/Module.cpp b/Engine/Plugins/Experimental/AnimNextChooser/Source/AnimNextChooser/Private/Module.cpp
new file mode 100644
--- /dev/null
+++ b/Engine/Plugins/Experimental/AnimNextChooser/Source/AnimNextChooser/Private/Module.cpp
@@ -0,0 +1,32 @@
+// Copyright Epic Games, Inc. All Rights Reserved.
+
+#include "CoreMinimal.h"
+#include "Module/AnimNextModule.h"
+#include "Modules/ModuleInterface.h"
+#include "Modules/ModuleManager.h"
+#include "RigVMCore/RigVMRegistry.h"
+#include "Animation/AnimSequence.h"
+#include "Chooser.h"
+
+#include "IAnimNextModuleInterface.h"
+#include "EvaluationVM/EvaluationVM.h"
+
+namespace UE::AnimNext::Chooser
+{
+
+class FModule : public IModuleInterface
+{
+public:
+	virtual void StartupModule() override
+	{
+	}
+
+	virtual void ShutdownModule() override
+	{
+	}
+
+};
+
+}
+
+IMPLEMENT_MODULE(UE::AnimNext::Chooser::FModule, AnimNextChooser)
diff --git a/Engine/Plugins/Experimental/AnimNextChooser/Source/AnimNextChooser/Private/RigUnit_EvaluateChooser.cpp b/Engine/Plugins/Experimental/AnimNextChooser/Source/AnimNextChooser/Private/RigUnit_EvaluateChooser.cpp
new file mode 100644
--- /dev/null
+++ b/Engine/Plugins/Experimental/AnimNextChooser/Source/AnimNextChooser/Private/RigUnit_EvaluateChooser.cpp
@@ -0,0 +1,22 @@
+// Copyright Epic Games, Inc. All Rights Reserved.
+
+
+#include "RigUnit_EvaluateChooser.h"
+
+
+FRigUnit_EvaluateChooser_Execute()
+{
+	Result = nullptr;
+
+	if (ContextObject && Chooser)
+	{
+		FChooserEvaluationContext ChooserContext;
+		ChooserContext.AddObjectParam(ContextObject);
+
+		UChooserTable::EvaluateChooser(ChooserContext, Chooser, FObjectChooserBase::FObjectChooserIteratorCallback::CreateLambda([&Result](UObject* InResult)
+			{
+				Result = InResult;
+				return FObjectChooserBase::EIteratorStatus::Stop;
+			}));
+	}
+}
diff --git a/Engine/Plugins/Experimental/AnimNextChooser/Source/AnimNextChooser/Private/RigUnit_EvaluateChooser.h b/Engine/Plugins/Experimental/AnimNextChooser/Source/AnimNextChooser/Private/RigUnit_EvaluateChooser.h
new file mode 100644
--- /dev/null
+++ b/Engine/Plugins/Experimental/AnimNextChooser/Source/AnimNextChooser/Private/RigUnit_EvaluateChooser.h
@@ -0,0 +1,38 @@
+// Copyright Epic Games, Inc. All Rights Reserved.
+
+
+#pragma once
+
+#include "CoreMinimal.h"
+#include "Chooser.h"
+#include "ControlRigDefines.h"
+#include "RigVMFunctions/Animation/RigVMFunction_AnimBase.h"
+#include "Units/RigUnit.h"
+
+#include "RigUnit_EvaluateChooser.generated.h"
+
+/*
+ * Evaluates a Chooser Table and outputs the selected UObject
+ */
+USTRUCT(meta = (DisplayName = "Evaluate Chooser", Keywords="Evaluate Chooser", Category = "Chooser", Varying, NodeColor="0.737911 0.099899 0.099899"))
+struct FRigUnit_EvaluateChooser : public FRigVMFunction_AnimBase
+{
+	GENERATED_BODY()
+
+	FRigUnit_EvaluateChooser()
+	{
+	}
+
+	/** Execute logic for this rig unit */
+	RIGVM_METHOD()
+	virtual void Execute() override;
+
+	UPROPERTY(meta = (Input))
+	TObjectPtr<UObject> ContextObject;
+	
+	UPROPERTY(EditAnywhere, Category = "Chooser", meta = (Input, Constant))
+	TObjectPtr<UChooserTable> Chooser;
+
+	UPROPERTY(meta = (Output))
+	TObjectPtr<UObject> Result;
+};
\ No newline at end of file
