AnimNext: Make chooser pinnable

diff --git a/Engine/Plugins/Experimental/AnimNextChooser/Source/AnimNextChooser/Private/RigUnit_EvaluateChooser.h b/Engine/Plugins/Experimental/AnimNextChooser/Source/AnimNextChooser/Private/RigUnit_EvaluateChooser.h
--- a/Engine/Plugins/Experimental/AnimNextChooser/Source/AnimNextChooser/Private/RigUnit_EvaluateChooser.h
+++ b/Engine/Plugins/Experimental/AnimNextChooser/Source/AnimNextChooser/Private/RigUnit_EvaluateChooser.h
@@ -22,7 +22,7 @@ struct FRigUnit_EvaluateChooser : public FRigVMStruct
 	UPROPERTY(meta = (Input))
 	TObjectPtr<UObject> ContextObject;
 	
-	UPROPERTY(EditAnywhere, Category = "Chooser", meta = (Input, Constant))
+	UPROPERTY(EditAnywhere, Category = "Chooser", meta = (Input))
 	TObjectPtr<UChooserTable> Chooser;
 
 	UPROPERTY(meta = (Output))
