Basic support for enum anim next chooser parameters

diff --git a/Engine/Plugins/Experimental/AnimNextChooser/Source/AnimNextChooser/Private/ChooserParameters.cpp b/Engine/Plugins/Experimental/AnimNextChooser/Source/AnimNextChooser/Private/ChooserParameters.cpp
--- a/Engine/Plugins/Experimental/AnimNextChooser/Source/AnimNextChooser/Private/ChooserParameters.cpp
+++ b/Engine/Plugins/Experimental/AnimNextChooser/Source/AnimNextChooser/Private/ChooserParameters.cpp
@@ -76,4 +76,31 @@ void FFloatAnimProperty::GetDisplayName(FText& OutName) const
 	}
 	
 	OutName = FText::FromString(DisplayName);
-}
\ No newline at end of file
+}
+
+bool FEnumAnimProperty::GetValue(FChooserEvaluationContext& Context, uint8& OutResult) const
+{
+	if(FAnimNextDataInterfaceInstance* Instance = UE::AnimNext::Private::GetFirstDataInterfaceInstance(Context))
+	{
+		return Instance->GetVariable(VariableName, OutResult) == EPropertyBagResult::Success;
+	}
+	return false;
+}
+
+bool FEnumAnimProperty::SetValue(FChooserEvaluationContext& Context, uint8 InValue) const
+{
+	return false;
+}
+
+void FEnumAnimProperty::GetDisplayName(FText& OutName) const
+{
+	FString DisplayName = VariableName.ToString();
+	int Index = -1;
+	DisplayName.FindLastChar(':', Index);
+	if (Index>=0)
+	{
+		DisplayName = DisplayName.RightChop(Index + 1);
+	}
+	
+	OutName = FText::FromString(DisplayName);
+}
diff --git a/Engine/Plugins/Experimental/AnimNextChooser/Source/AnimNextChooser/Private/ChooserParameters.h b/Engine/Plugins/Experimental/AnimNextChooser/Source/AnimNextChooser/Private/ChooserParameters.h
--- a/Engine/Plugins/Experimental/AnimNextChooser/Source/AnimNextChooser/Private/ChooserParameters.h
+++ b/Engine/Plugins/Experimental/AnimNextChooser/Source/AnimNextChooser/Private/ChooserParameters.h
@@ -4,6 +4,7 @@
 #include "CoreMinimal.h"
 #include "IChooserParameterBool.h"
 #include "IChooserParameterFloat.h"
+#include "IChooserParameterEnum.h"
 #include "ChooserParameters.generated.h"
 
 USTRUCT(DisplayName = "Bool Anim Param")
@@ -32,4 +33,26 @@ public:
 	virtual bool GetValue(FChooserEvaluationContext& Context, double& OutResult) const override;
 	virtual bool SetValue(FChooserEvaluationContext& Context, double InValue) const override;
 	virtual void GetDisplayName(FText& OutName) const override;
+};
+
+USTRUCT(DisplayName = "Enum Anim Param")
+struct FEnumAnimProperty :  public FChooserParameterEnumBase
+{
+	GENERATED_BODY()
+public:
+
+	UPROPERTY(DisplayName = "Variable", EditAnywhere, Category="Variable")
+	FName VariableName;
+
+	UPROPERTY(EditAnywhere, Category="Variable")
+	TObjectPtr<UEnum> Enum;
+	
+	virtual bool GetValue(FChooserEvaluationContext& Context, uint8& OutResult) const override;
+	virtual bool SetValue(FChooserEvaluationContext& Context, uint8 InValue) const override;
+
+	virtual void GetDisplayName(FText& OutName) const override;
+
+	#if WITH_EDITOR
+	virtual const UEnum* GetEnum() const override { return Enum; }
+	#endif
 };
\ No newline at end of file
