Add array output support to evaluate chooser node in rigvm

diff --git a/Engine/Plugins/Experimental/AnimNextChooser/Source/AnimNextChooser/Private/RigUnit_EvaluateChooser.cpp b/Engine/Plugins/Experimental/AnimNextChooser/Source/AnimNextChooser/Private/RigUnit_EvaluateChooser.cpp
--- a/Engine/Plugins/Experimental/AnimNextChooser/Source/AnimNextChooser/Private/RigUnit_EvaluateChooser.cpp
+++ b/Engine/Plugins/Experimental/AnimNextChooser/Source/AnimNextChooser/Private/RigUnit_EvaluateChooser.cpp
@@ -5,6 +5,8 @@
 #include "DataInterface/AnimNextDataInterfaceInstance.h"
 #include "ControlRig.h"
 
+#include UE_INLINE_GENERATED_CPP_BY_NAME(RigUnit_EvaluateChooser)
+
 static void RunChooserHelper(TConstArrayView<TObjectPtr<UObject>> ContextObjects, FStructView ContextStruct, TObjectPtr<UChooserTable> Chooser, TObjectPtr<UObject>& OutResult)
 {
 	if ((!ContextObjects.IsEmpty() || ContextStruct.IsValid()) && Chooser != nullptr)
@@ -43,4 +45,115 @@ FRigUnit_EvaluateChooser_AnimNext_Execute()
 
 	const FAnimNextDataInterfaceInstance& Instance = ExecuteContext.GetContextData<FAnimNextModuleContextData>().GetDataInterfaceInstance();
 	RunChooserHelper({ ContextObject }, FStructView::Make(const_cast<FAnimNextDataInterfaceInstance&>(Instance)), Chooser, Result);
-}
\ No newline at end of file
+}
+
+const FName FRigVMDispatch_EvaluateChooser::ResultName = TEXT("Result");
+const FName FRigVMDispatch_EvaluateChooser::ChooserName = TEXT("Chooser");
+const FName FRigVMDispatch_EvaluateChooser::ContextObjectName = TEXT("ContextObject");
+
+FRigVMDispatch_EvaluateChooser::FRigVMDispatch_EvaluateChooser()
+{
+	FactoryScriptStruct = StaticStruct();
+}
+
+FName FRigVMDispatch_EvaluateChooser::GetArgumentNameForOperandIndex(int32 InOperandIndex, int32 InTotalOperands) const
+{
+	static const FName ArgumentNames[] =
+	{
+		ChooserName,
+		ContextObjectName,
+		ResultName,
+	};
+	check(InTotalOperands == UE_ARRAY_COUNT(ArgumentNames));
+	return ArgumentNames[InOperandIndex];
+}
+
+const TArray<FRigVMTemplateArgumentInfo>& FRigVMDispatch_EvaluateChooser::GetArgumentInfos() const
+{
+	static TArray<FRigVMTemplateArgumentInfo> Infos;
+	if(Infos.IsEmpty())
+	{
+		static const TArray<FRigVMTemplateArgument::ETypeCategory> ResultCategories =
+		{
+			FRigVMTemplateArgument::ETypeCategory_SingleObjectValue,
+			FRigVMTemplateArgument::ETypeCategory_ArrayObjectValue
+		};
+		
+		const FRigVMRegistry_NoLock& Registry = FRigVMRegistry_NoLock::GetForRead();
+		const TRigVMTypeIndex ChooserTypeIndex = Registry.GetTypeIndex_NoLock<UChooserTable>();
+		const TRigVMTypeIndex ObjectTypeIndex = Registry.GetTypeIndex_NoLock<UObject>();
+		Infos.Emplace(ChooserName, ERigVMPinDirection::Input, ChooserTypeIndex);
+		Infos.Emplace(ContextObjectName, ERigVMPinDirection::Input, ObjectTypeIndex);
+		Infos.Emplace(ResultName, ERigVMPinDirection::Output, ResultCategories);
+	}
+
+	return Infos;
+}
+
+void FRigVMDispatch_EvaluateChooser::Execute(FRigVMExtendedExecuteContext& InContext, FRigVMMemoryHandleArray Handles, FRigVMPredicateBranchArray RigVMBranches)
+{
+
+	const FProperty* ResultProperty = Handles[2].GetResolvedProperty();
+	
+	TObjectPtr<UChooserTable>* Chooser = reinterpret_cast<TObjectPtr<UChooserTable>*>(Handles[0].GetData());
+	TObjectPtr<UObject>* ContextObject = reinterpret_cast<TObjectPtr<UObject>*>(Handles[1].GetData());
+
+
+	FChooserEvaluationContext ChooserContext;
+	if (InContext.GetContextPublicDataStruct()->IsChildOf(FAnimNextExecuteContext::StaticStruct()))
+	{
+		FAnimNextExecuteContext& VMExecuteContext = InContext.GetPublicData<FAnimNextExecuteContext>();
+		const FAnimNextDataInterfaceInstance& Instance = VMExecuteContext.GetContextData<FAnimNextModuleContextData>().GetDataInterfaceInstance();
+		ChooserContext.AddStructViewParam(	FStructView::Make(const_cast<FAnimNextDataInterfaceInstance&>(Instance)));
+	}
+	else
+	{
+		FControlRigExecuteContext& ControlRigContext = InContext.GetPublicDataSafe<FControlRigExecuteContext>();
+		ChooserContext.AddObjectParam(ControlRigContext.ControlRig);
+	}
+	ChooserContext.AddObjectParam(*ContextObject);
+
+	if (const FArrayProperty* ArrayResultProperty = CastField<FArrayProperty>(ResultProperty))
+	{
+		if (const FObjectProperty* ObjectResultProperty = CastField<FObjectProperty>(ArrayResultProperty->Inner))
+		{
+			UClass* ObjectClass = ObjectResultProperty->PropertyClass;
+			FScriptArrayHelper ArrayHelper(ArrayResultProperty, Handles[2].GetData());
+			ArrayHelper.Resize(0);
+			
+			UChooserTable::EvaluateChooser(ChooserContext, *Chooser, FObjectChooserBase::FObjectChooserIteratorCallback::CreateLambda([ObjectClass, &ArrayHelper](UObject* InResult)
+			{
+				if (InResult)
+				{
+					if (InResult->GetClass()->IsChildOf(ObjectClass))
+					{
+						int Index = ArrayHelper.AddValue();
+						TObjectPtr<UObject>* Result = (TObjectPtr<UObject>*)ArrayHelper.GetElementPtr(Index);
+						*Result = InResult;
+						return Index == 0 ? FObjectChooserBase::EIteratorStatus::ContinueWithOutputs : FObjectChooserBase::EIteratorStatus::Continue;
+					}
+				}
+				return FObjectChooserBase::EIteratorStatus::Continue;
+			}));
+		}
+	}
+	else if (const FObjectProperty* ObjectResultProperty = CastField<FObjectProperty>(ResultProperty))
+	{
+		UClass* ObjectClass = ObjectResultProperty->PropertyClass;
+		TObjectPtr<UObject>* Result = (TObjectPtr<UObject>*)Handles[2].GetData();
+		UChooserTable::EvaluateChooser(ChooserContext, *Chooser, FObjectChooserBase::FObjectChooserIteratorCallback::CreateLambda([ObjectClass, &Result](UObject* InResult)
+		{
+			if (InResult)
+			{
+				if (InResult->GetClass()->IsChildOf(ObjectClass))
+				{
+					*Result = InResult;
+					return FObjectChooserBase::EIteratorStatus::Stop;
+				}
+			}
+			
+			return FObjectChooserBase::EIteratorStatus::Continue;
+		}));
+	}
+}
+
diff --git a/Engine/Plugins/Experimental/AnimNextChooser/Source/AnimNextChooser/Private/RigUnit_EvaluateChooser.h b/Engine/Plugins/Experimental/AnimNextChooser/Source/AnimNextChooser/Private/RigUnit_EvaluateChooser.h
--- a/Engine/Plugins/Experimental/AnimNextChooser/Source/AnimNextChooser/Private/RigUnit_EvaluateChooser.h
+++ b/Engine/Plugins/Experimental/AnimNextChooser/Source/AnimNextChooser/Private/RigUnit_EvaluateChooser.h
@@ -7,14 +7,48 @@
 #include "AnimNextExecuteContext.h"
 #include "Chooser.h"
 #include "ControlRigDefines.h"
+#include "RigVMCore/RigVMDispatchFactory.h"
 #include "Units/RigUnit.h"
 
 #include "RigUnit_EvaluateChooser.generated.h"
 
+
+/*
+ * Evaluates a chooser table and outputs either an object or an array of objects depending on what the result pin is connected to
+ * Compatible with both ControlRig and AnimNext graphs
+ */
+USTRUCT(meta = (DisplayName = "Evaluate Chooser", Category="Choosers", NodeColor = "0.8, 0, 0.2, 1"))
+struct FRigVMDispatch_EvaluateChooser : public FRigVMDispatchFactory
+{
+	GENERATED_BODY()
+
+	FRigVMDispatch_EvaluateChooser();
+	
+private:
+	friend struct UE::AnimNext::UncookedOnly::FUtils;
+
+	virtual UScriptStruct* GetExecuteContextStruct() const { return FRigVMExecuteContext::StaticStruct(); }
+	virtual FName GetArgumentNameForOperandIndex(int32 InOperandIndex, int32 InTotalOperands) const override;
+	virtual const TArray<FRigVMTemplateArgumentInfo>& GetArgumentInfos() const override;
+
+	virtual bool IsSingleton() const override { return true; }
+
+	virtual FRigVMFunctionPtr GetDispatchFunctionImpl(const FRigVMTemplateTypeMap& InTypes) const override
+	{
+		return &FRigVMDispatch_EvaluateChooser::Execute;
+	}
+	
+	static void Execute(FRigVMExtendedExecuteContext& InContext, FRigVMMemoryHandleArray Handles, FRigVMPredicateBranchArray RigVMBranches);
+	
+	static const FName ChooserName;
+	static const FName ContextObjectName;
+	static const FName ResultName;
+};
+
 /*
  * Evaluates a Chooser Table and outputs the selected UObject
  */
-USTRUCT(meta = (Abstract, Keywords="Evaluate Chooser", Category = "Chooser", Varying, NodeColor="0.737911 0.099899 0.099899"))
+USTRUCT(meta = (Abstract, Varying))
 struct FRigUnit_EvaluateChooser : public FRigVMStruct
 {
 	GENERATED_BODY()
@@ -31,8 +65,9 @@ struct FRigUnit_EvaluateChooser : public FRigVMStruct
 
 /*
  * Evaluates a Chooser Table in the context of ControlRig
+ * deprecated in favor of RigVMDispatch_EvaluateChooser
  */
-USTRUCT(meta = (DisplayName = "Evaluate Chooser", Varying, NodeColor="0.737911 0.099899 0.099899", ExecuteContext="FControlRigExecuteContext"))
+USTRUCT(meta = (ExecuteContext="FControlRigExecuteContext"))
 struct FRigUnit_EvaluateChooser_ControlRig : public FRigUnit_EvaluateChooser
 {
 	GENERATED_BODY()
@@ -44,8 +79,9 @@ struct FRigUnit_EvaluateChooser_ControlRig : public FRigUnit_EvaluateChooser
 
 /*
  * Evaluates a Chooser Table in the context of AnimNext
+ * deprecated in favor of RigVMDispatch_EvaluateChooser
  */
-USTRUCT(meta = (DisplayName = "Evaluate Chooser", Varying, NodeColor="0.737911 0.099899 0.099899", ExecuteContext="FAnimNextExecuteContext"))
+USTRUCT(meta = (ExecuteContext="FAnimNextExecuteContext"))
 struct FRigUnit_EvaluateChooser_AnimNext : public FRigUnit_EvaluateChooser
 {
 	GENERATED_BODY()
