Add SetValue implementations to Anim Next chooser parameter implementations
#ForUE5

diff --git a/Engine/Plugins/Experimental/AnimNextChooser/Source/AnimNextChooser/Private/ChooserParameters.cpp b/Engine/Plugins/Experimental/AnimNextChooser/Source/AnimNextChooser/Private/ChooserParameters.cpp
--- a/Engine/Plugins/Experimental/AnimNextChooser/Source/AnimNextChooser/Private/ChooserParameters.cpp
+++ b/Engine/Plugins/Experimental/AnimNextChooser/Source/AnimNextChooser/Private/ChooserParameters.cpp
@@ -34,6 +34,10 @@ bool FBoolAnimProperty::GetValue(FChooserEvaluationContext& Context, bool& OutRe
 
 bool FBoolAnimProperty::SetValue(FChooserEvaluationContext& Context, bool InValue) const
 {
+	if(FAnimNextDataInterfaceInstance* Instance = UE::AnimNext::Private::GetFirstDataInterfaceInstance(Context))
+	{
+		return Instance->SetVariable(VariableName, InValue) == EPropertyBagResult::Success;
+	}
 	return false;
 }
 
@@ -62,7 +66,11 @@ bool FFloatAnimProperty::GetValue(FChooserEvaluationContext& Context, double& Ou
 
 bool FFloatAnimProperty::SetValue(FChooserEvaluationContext& Context, double InValue) const
 {
-	return false;
+	if(FAnimNextDataInterfaceInstance* Instance = UE::AnimNext::Private::GetFirstDataInterfaceInstance(Context))
+ 	{
+ 		return Instance->SetVariable(VariableName, InValue) == EPropertyBagResult::Success;
+ 	}
+ 	return false;
 }
 
 void FFloatAnimProperty::GetDisplayName(FText& OutName) const
@@ -89,6 +97,10 @@ bool FEnumAnimProperty::GetValue(FChooserEvaluationContext& Context, uint8& OutR
 
 bool FEnumAnimProperty::SetValue(FChooserEvaluationContext& Context, uint8 InValue) const
 {
+	if(FAnimNextDataInterfaceInstance* Instance = UE::AnimNext::Private::GetFirstDataInterfaceInstance(Context))
+	{
+		return Instance->SetVariable(VariableName, InValue) == EPropertyBagResult::Success;
+	}
 	return false;
 }
 
