Access AnimNextComponent object from rigvm

diff --git a/Engine/Plugins/Experimental/AnimNextChooser/Source/AnimNextChooser/Private/RigUnit_OwningObject.cpp b/Engine/Plugins/Experimental/AnimNextChooser/Source/AnimNextChooser/Private/RigUnit_OwningObject.cpp
new file mode 100644
--- /dev/null
+++ b/Engine/Plugins/Experimental/AnimNextChooser/Source/AnimNextChooser/Private/RigUnit_OwningObject.cpp
@@ -0,0 +1,11 @@
+// Copyright Epic Games, Inc. All Rights Reserved.
+
+
+#include "RigUnit_OwningObject.h"
+
+#include UE_INLINE_GENERATED_CPP_BY_NAME(RigUnit_OwningObject)
+
+FRigUnit_OwningObject_Execute()
+{
+	Result = const_cast<UObject*>(ExecuteContext.GetOwningObject());
+}
diff --git a/Engine/Plugins/Experimental/AnimNextChooser/Source/AnimNextChooser/Private/RigUnit_OwningObject.h b/Engine/Plugins/Experimental/AnimNextChooser/Source/AnimNextChooser/Private/RigUnit_OwningObject.h
new file mode 100644
--- /dev/null
+++ b/Engine/Plugins/Experimental/AnimNextChooser/Source/AnimNextChooser/Private/RigUnit_OwningObject.h
@@ -0,0 +1,24 @@
+// Copyright Epic Games, Inc. All Rights Reserved.
+
+
+#pragma once
+
+#include "CoreMinimal.h"
+#include "ControlRigDefines.h"
+#include "Units/RigUnit.h"
+
+#include "RigUnit_OwningObject.generated.h"
+
+
+USTRUCT(meta = (DisplayName="Owning Object"))
+struct FRigUnit_OwningObject : public FRigVMStruct
+{
+	GENERATED_BODY()
+	
+	UPROPERTY(meta = (Output))
+	TObjectPtr<UObject> Result;
+
+	/** Execute logic for this rig unit */
+	RIGVM_METHOD()
+	virtual void Execute() override;
+};
\ No newline at end of file
